Size Map: Fixed assets with back-references causing the asset to appear in the "shared" node

diff --git a/Engine/Source/Editor/SizeMap/SSizeMap.cpp b/Engine/Source/Editor/SizeMap/SSizeMap.cpp
--- a/Engine/Source/Editor/SizeMap/SSizeMap.cpp
+++ b/Engine/Source/Editor/SizeMap/SSizeMap.cpp
@@ -170,31 +170,46 @@ void SSizeMap::GatherDependenciesRecursively( FAssetRegistryModule& AssetRegistr
 		{ 
 			// OK, we've determined that this asset has already been referenced by something else in our tree.  We'll move it to a "shared" group
 			// so all of the assets that are referenced in multiple places can be seen together.
-			TSharedPtr<FTreeMapNodeData> ExistingNodeInTree = VisitedAssetPackageNames[ AssetPackageName ];
+			TSharedPtr<FTreeMapNodeData> ExistingNode = VisitedAssetPackageNames[ AssetPackageName ];
 
 			// Is the existing node not already under the "shared" group?  Note that it might still be (indirectly) under
 			// the "shared" group, in which case we'll still want to move it up to the root since we've figured out that it is
 			// actually shared between multiple assets which themselves may be shared
-			if( ExistingNodeInTree->Parent != SharedRootNode.Get() )
+			if( ExistingNode->Parent != SharedRootNode.Get() )
 			{
 				// Don't bother moving any of the assets at the root level into a "shared" bucket.  We're only trying to best
 				// represent the memory used when all of the root-level assets have become loaded.  It's OK if root-level assets
 				// are referenced by other assets in the set -- we don't need to indicate they are shared explicitly
-				if( ExistingNodeInTree->Parent->Parent != nullptr )
+				FTreeMapNodeData* ExistingNodeParent = ExistingNode->Parent;
+				check( ExistingNodeParent != nullptr );
+				const bool bExistingNodeIsAtRootLevel = ExistingNodeParent->Parent == nullptr;
+				if( !bExistingNodeIsAtRootLevel )
 				{
+					// OK, current asset (AssetPackageName) is definitely not a root level asset, but its already in the tree
+					// somewhere as a non-shared, non-root level asset.  We need to make sure that this Node's reference is not from the
+					// same root-level asset as the ExistingNodeInTree.  Otherwise, there's no need to move it to a 'shared' group.
 					FTreeMapNodeData* MyParentNode = Node.Get();
-
+					check( MyParentNode != nullptr );
+					FTreeMapNodeData* MyRootLevelAssetNode = MyParentNode;
+					while( MyRootLevelAssetNode->Parent != nullptr && MyRootLevelAssetNode->Parent->Parent != nullptr )
+					{
+						MyRootLevelAssetNode = MyRootLevelAssetNode->Parent;
+					}
+					if( MyRootLevelAssetNode->Parent == nullptr )
+					{
+						// No root asset (Node must be a root level asset itself!)
+						MyRootLevelAssetNode = nullptr;
+					}
+					
 					// Find the existing node's root level asset node
-					FTreeMapNodeData* ExistingNodeRootLevelAssetNode = ExistingNodeInTree.Get();
+					FTreeMapNodeData* ExistingNodeRootLevelAssetNode = ExistingNodeParent;
 					while( ExistingNodeRootLevelAssetNode->Parent->Parent != nullptr )
 					{
 						ExistingNodeRootLevelAssetNode = ExistingNodeRootLevelAssetNode->Parent;
 					}
 
-					// If we're being referenced by another node within the same asset, no need to move it to a 'shared' group.  We only care
-					// about referenced assets that are shared between different assets
-					// @todo sizemap urgent: Doesn't work correctly with all assets (BP_Sky)
-					if( MyParentNode != ExistingNodeRootLevelAssetNode->Parent )
+					// If we're being referenced by another node within the same asset, no need to move it to a 'shared' group.  
+					if( MyRootLevelAssetNode != ExistingNodeRootLevelAssetNode )
 					{
 						// This asset was already referenced by something else (or was in our top level list of assets to display sizes for)
 						if( !SharedRootNode.IsValid() )
@@ -212,9 +227,9 @@ void SSizeMap::GatherDependenciesRecursively( FAssetRegistryModule& AssetRegistr
 						}
 
 						// Reparent the node that we've now determined to be shared
-						ExistingNodeInTree->Parent->Children.Remove( ExistingNodeInTree );
-						SharedRootNode->Children.Add( ExistingNodeInTree );
-						ExistingNodeInTree->Parent = SharedRootNode.Get();
+						ExistingNode->Parent->Children.Remove( ExistingNode );
+						SharedRootNode->Children.Add( ExistingNode );
+						ExistingNode->Parent = SharedRootNode.Get();
 					}
 				}
 			}
