Size Mapper fixes
- Fixed references not being up-to-date as changes are made live in the editor.  Simply close and reopen the Size Map to see updated references. (UE-14410)
- Fixed many issues with Size Map not showing the correct set of referenced content by an asset (no longer uses Asset Registry, instead serializes to find references.)
- Fixed 2D textures not reporting correct size (they were reporting current streamed-in mips, not worst-case size)
- Fixed assets that reported zero bytes as their size not being tracked as "unknown size"
- Fixed root-level assets sometimes getting bundled into the "SHARED" group
- Fixed Size Map no longer functioning when viewing large numbers of assets, such as all Engine Content
- Fixed Size Map having to be closed and reopened sometimes to show all of the correct references
- Fixed node text sometimes drawing with unintentionally huge fonts due to integer overflow
- Fixed the asset type name not being shown next to an asset name that had child references
- Fixed hard-to-read text labels due to node colors being desaturated too much at even shallow depth levels

diff --git a/Engine/Source/Developer/TreeMap/STreeMap.cpp b/Engine/Source/Developer/TreeMap/STreeMap.cpp
--- a/Engine/Source/Developer/TreeMap/STreeMap.cpp
+++ b/Engine/Source/Developer/TreeMap/STreeMap.cpp
@@ -1196,7 +1196,7 @@ void STreeMap::ApplyVisualizationToNodesRecursively( const FTreeMapNodeDataRef&
 	// Select default color saturation based on tree depth
 	FLinearColor MyDefaultColor = DefaultColor;
 	auto HSV = MyDefaultColor.LinearRGBToHSV();
-	const float SaturationReductionPerDepthLevel = 0.2f;	// @todo treemap custom: Make configurable in the tree map settings?  Calculate tree depth?
+	const float SaturationReductionPerDepthLevel = 0.05f;	// @todo treemap custom: Make configurable in the tree map settings?  Calculate tree depth?
 	const float MinAllowedSaturation = 0.1f;
 	HSV.G = FMath::Max( MinAllowedSaturation, HSV.G - TreeDepth * SaturationReductionPerDepthLevel );
 	MyDefaultColor = HSV.HSVToLinearRGB();
diff --git a/Engine/Source/Developer/TreeMap/TreeMap.cpp b/Engine/Source/Developer/TreeMap/TreeMap.cpp
--- a/Engine/Source/Developer/TreeMap/TreeMap.cpp
+++ b/Engine/Source/Developer/TreeMap/TreeMap.cpp
@@ -466,11 +466,11 @@ void FTreeMap::PadNodesRecursively( const FTreeMapOptions& Options, const FTreeM
 	// Choose a height for this node's font
 	const uint16 MinAllowedFontSize = 8;		// @todo treemap custom: Don't hardcode and instead make this a customizable option?
 	Node->NameFont = Options.NameFont; 
-	Node->NameFont.Size = FMath::Max< uint16 >( MinAllowedFontSize, Options.NameFont.Size - ( TreeDepth * Options.FontSizeChangeBasedOnDepth ) );
+	Node->NameFont.Size = FMath::Max< int32 >( MinAllowedFontSize, Options.NameFont.Size - ( TreeDepth * Options.FontSizeChangeBasedOnDepth ) );
 	Node->Name2Font = Options.Name2Font; 
-	Node->Name2Font.Size = FMath::Max< uint16 >( MinAllowedFontSize, Options.Name2Font.Size - ( TreeDepth * Options.FontSizeChangeBasedOnDepth ) );
+	Node->Name2Font.Size = FMath::Max< int32 >( MinAllowedFontSize, Options.Name2Font.Size - ( TreeDepth * Options.FontSizeChangeBasedOnDepth ) );
 	Node->CenterTextFont = Options.CenterTextFont; 
-	Node->CenterTextFont.Size = FMath::Max< uint16 >( MinAllowedFontSize, Options.CenterTextFont.Size - ( TreeDepth * Options.FontSizeChangeBasedOnDepth ) );
+	Node->CenterTextFont.Size = FMath::Max< int32 >( MinAllowedFontSize, Options.CenterTextFont.Size - ( TreeDepth * Options.FontSizeChangeBasedOnDepth ) );
 
 	if( Node != RootNode )
 	{
diff --git a/Engine/Source/Editor/SizeMap/SSizeMap.cpp b/Engine/Source/Editor/SizeMap/SSizeMap.cpp
--- a/Engine/Source/Editor/SizeMap/SSizeMap.cpp
+++ b/Engine/Source/Editor/SizeMap/SSizeMap.cpp
@@ -70,46 +70,61 @@ namespace SizeMapInternals
 	{
 	public:
 		FAssetReferenceFinder( UObject* Object )
-			: CurrentObject( Object )
 		{
 			ArIsObjectReferenceCollector = true;
 			ArIgnoreOuterRef = true;
 
-			CurrentObject->Serialize( *this );
+			check( Object != nullptr );
+			AllVisitedObjects.Add( Object );
+			Object->Serialize( *this );
 		}
 
 		FArchive& operator<<( UObject*& Object )
 		{
 			// Only look at objects which are valid
-			const bool bValidObject =
-				Object &&	// Object should not be NULL
-				!Object->HasAnyFlags( RF_Transient | RF_PendingKill ) && // Should not be transient or pending kill
-				Object->HasAllFlags( RF_Public ) &&	// All assets should be public
-				Object->GetOuter()->IsA( UPackage::StaticClass() ) &&	// Only want outer assets (these should be the only public assets, anyway)
-				( Cast<AActor>( Object ) == nullptr ) &&			// skip actors
-				( Cast<UActorComponent>( Object ) == nullptr ) &&	// skip components
-				( Cast<UClass>( Object ) == nullptr );				// skip UClasses
-
-			if( bValidObject )
+			const bool bIsValidObject =
+				Object != nullptr &&	// Object should not be null
+				!Object->HasAnyFlags( RF_Transient | RF_PendingKill );	// Should not be transient or pending kill
+			if( bIsValidObject )
 			{
-				// Ignore self referencing objects
-				if( Object != CurrentObject )
+				// Skip objects that we've already processed
+				if( !AllVisitedObjects.Contains( Object ) )
 				{
-					// ADD REFERENCE(CurrentObject DEPENDS on Object)
-					ReferencedObjects.Add( Object );
+					AllVisitedObjects.Add( Object );
+
+					const bool bIsAsset =
+						Object->GetOuter()->IsA( UPackage::StaticClass() ) &&	// Only want outer assets (these should be the only public assets, anyway)
+						Object->HasAllFlags( RF_Public );						// Assets should be public
+
+					if( bIsAsset )
+					{
+						ReferencedAssets.Add( Object );
+					}
+					else
+					{
+						// It's probably an inner object.  Recursively serialize.
+						Object->Serialize( *this );
+
+						// Make sure the object's class is serialized too, so that we catch any assets referenced from the class defaults
+						AllVisitedObjects.Add( Object->GetClass() );
+						Object->GetClass()->Serialize( *this );	// @todo sizemap urgent: Doesn't really seem to be needed
+					}
 				}
 			}
-
 			return *this;
 		}
 
-	private:
-		/** The object currently being serialized. */
-		UObject* CurrentObject;
+		TSet< UObject* >& GetReferencedAssets()
+		{
+			return ReferencedAssets;
+		}
+
+	protected:
+		/** The set of referenced assets */
+		TSet< UObject* > ReferencedAssets;
 
-	public:
-		/** The set of referenced objects */
-		TSet< UObject* > ReferencedObjects;
+		/** Set of all objects we've visited, so we don't follow cycles */
+		TSet< UObject* > AllVisitedObjects;
 	};
 
 
@@ -160,11 +175,11 @@ namespace SizeMapInternals
 }
 
 
-void SSizeMap::GatherDependenciesRecursively( FAssetRegistryModule& AssetRegistryModule, TSharedPtr<FAssetThumbnailPool>& InAssetThumbnailPool, TMap<FName, TSharedPtr<FTreeMapNodeData>>& VisitedAssetPackageNames, const TArray<FName>& AssetPackageNames, const TSharedPtr<FTreeMapNodeData>& Node, TSharedPtr<FTreeMapNodeData>& SharedRootNode, int32& NumAssetsWhichFailedToLoad )
+void SSizeMap::GatherDependenciesRecursively( FAssetRegistryModule& AssetRegistryModule, TSharedPtr<FAssetThumbnailPool>& InAssetThumbnailPool, TMap<FName, TSharedPtr<FTreeMapNodeData>>& VisitedAssetPackageNames, const TArray<FName>& RootAssetPackageNames, const TArray<FName>& AssetPackageNames, const TSharedPtr<FTreeMapNodeData>& Node, TSharedPtr<FTreeMapNodeData>& SharedRootNode, int32& NumAssetsWhichFailedToLoad )
 {
 	for( const FName AssetPackageName : AssetPackageNames )
 	{
-		// Have we already added this asset to the tree?  If so, we'll either move it to a "shared" group or (if its referenced again by the same
+		// Have we already added this asset to the tree?  If so, we'll either move it to a "shared" group or (if it's referenced again by the same
 		// root-level asset) ignore it
 		if( VisitedAssetPackageNames.Contains( AssetPackageName ) )
 		{ 
@@ -182,10 +197,10 @@ void SSizeMap::GatherDependenciesRecursively( FAssetRegistryModule& AssetRegistr
 				// are referenced by other assets in the set -- we don't need to indicate they are shared explicitly
 				FTreeMapNodeData* ExistingNodeParent = ExistingNode->Parent;
 				check( ExistingNodeParent != nullptr );
-				const bool bExistingNodeIsAtRootLevel = ExistingNodeParent->Parent == nullptr;
+				const bool bExistingNodeIsAtRootLevel = ExistingNodeParent->Parent == nullptr || RootAssetPackageNames.Contains( AssetPackageName );
 				if( !bExistingNodeIsAtRootLevel )
 				{
-					// OK, current asset (AssetPackageName) is definitely not a root level asset, but its already in the tree
+					// OK, the current asset (AssetPackageName) is definitely not a root level asset, but its already in the tree
 					// somewhere as a non-shared, non-root level asset.  We need to make sure that this Node's reference is not from the
 					// same root-level asset as the ExistingNodeInTree.  Otherwise, there's no need to move it to a 'shared' group.
 					FTreeMapNodeData* MyParentNode = Node.Get();
@@ -260,13 +275,13 @@ void SSizeMap::GatherDependenciesRecursively( FAssetRegistryModule& AssetRegistr
 				NodeSizeMapData.bHasKnownSize = false;
 
 				// Find the asset using the asset registry
-				// @todo sizemap urgent: Asset registry is faster but possibly not as exhaustive (no PostLoad created references, etc.)  Maybe should be optional?
-				// @todo sizemap urgent: When not using asset registry to find references, all references of map files are not showing up
+				// @todo sizemap: Asset registry-based reference gathering is faster but possibly not as exhaustive (no PostLoad created references, etc.)  Maybe should be optional?
+				// @todo sizemap: With AR-based reference gathering, sometimes the size map is missing root level dependencies until you reopen it a few times (Buggy BP)
+				// @todo sizemap: With AR-based reference gathering, reference changes at editor-time do not appear in the Size Map until you restart
+				// @todo sizemap: With AR-based reference gathering, opening the size map for all engine content caused the window to not respond until a restart
 				// @todo sizemap: We don't really need the asset registry given we need to load the objects to figure out their size, unless we make that AR-searchable.
 				//   ---> This would allow us to not have to wait for AR initialization.  But if we made size AR-searchable, we could run very quickly for large data sets!
-				// @todo sizemap urgent: When loading dependencies using the FAssetReferrenceFinder, we never find out about any missing references (unlike with the asset
-				//    registry approach, where we can tell when something is missing.)  This could product misleading results!
-				const bool bUseAssetRegistryForDependencies = true;
+				const bool bUseAssetRegistryForDependencies = false;
 
 				const FString AssetPathString = AssetPackageNameString + TEXT(".") + FPackageName::GetLongPackageAssetName( AssetPackageNameString );
 				const FAssetData FoundAssetData = AssetRegistryModule.Get().GetAssetByObjectPath( FName( *AssetPathString ) );
@@ -287,24 +302,36 @@ void SSizeMap::GatherDependenciesRecursively( FAssetRegistryModule& AssetRegistr
 						else
 						{
 							SizeMapInternals::FAssetReferenceFinder References( Asset );
-							for( UObject* Object : References.ReferencedObjects )
+							for( UObject* Object : References.GetReferencedAssets() )
 							{
 								ReferencedAssetPackageNames.Add( FName( *Object->GetOutermost()->GetPathName() ) );
 							}
 						}
 
-						NodeSizeMapData.AssetSize = Asset->GetResourceSize( EResourceSizeMode::Exclusive );
-						NodeSizeMapData.bHasKnownSize = ( NodeSizeMapData.AssetSize != UObject::RESOURCE_SIZE_NONE );
+						// For textures, make sure we're getting the worst case size, not the size of the currently loaded set of mips
+						// @todo sizemap: We should instead have a special EResourceSizeMode that asks for the worst case size.  Some assets (like UTextureCube) currently always report resident mip size, even when asked for inclusive size
+						if( Asset->IsA( UTexture2D::StaticClass() ) )
+						{
+							NodeSizeMapData.AssetSize = Asset->GetResourceSize( EResourceSizeMode::Inclusive );
+						}
+						else
+						{
+							NodeSizeMapData.AssetSize = Asset->GetResourceSize( EResourceSizeMode::Exclusive );
+						}
+
+						NodeSizeMapData.bHasKnownSize = NodeSizeMapData.AssetSize != UObject::RESOURCE_SIZE_NONE && NodeSizeMapData.AssetSize != 0;
 						if( !NodeSizeMapData.bHasKnownSize )
 						{
 							// Asset has no meaningful size
 							NodeSizeMapData.AssetSize = 0;
+
+							// @todo sizemap urgent: Try to serialize to figure out how big it is (not into sub-assets though!)
+							// FObjectMemoryAnalyzer ObjectMemoryAnalyzer( Asset );
 						}
 
 								
 						// Now visit all of the assets that we are referencing
-						GatherDependenciesRecursively( AssetRegistryModule, InAssetThumbnailPool, VisitedAssetPackageNames, ReferencedAssetPackageNames, ChildTreeMapNode, SharedRootNode, NumAssetsWhichFailedToLoad );
-
+						GatherDependenciesRecursively( AssetRegistryModule, InAssetThumbnailPool, VisitedAssetPackageNames, RootAssetPackageNames, ReferencedAssetPackageNames, ChildTreeMapNode, SharedRootNode, NumAssetsWhichFailedToLoad );
 					}
 					else
 					{
@@ -345,7 +372,7 @@ void SSizeMap::FinalizeNodesRecursively( TSharedPtr<FTreeMapNodeData>& Node, con
 	{
 		// @todo sizemap: Should we indicate in a non-shared parent node how many if its dependents ended up being in the "shared" bucket?  Probably 
 		// not that important, because the user can choose to view that asset in isolation to see the full tree.
-		Node->Name = FString::Printf( TEXT( "%s (%s)" ),
+		Node->Name = FString::Printf( TEXT( "%s  (%s)" ),
 			*LOCTEXT( "SharedGroupName", "*SHARED*" ).ToString(),
 			*SizeMapInternals::MakeBestSizeString( SubtreeSize, !bAnyUnknownSizes ) );
 
@@ -411,14 +438,15 @@ void SSizeMap::FinalizeNodesRecursively( TSharedPtr<FTreeMapNodeData>& Node, con
 			// Container nodes are always auto-sized
 			Node->Size = 0.0f;
 
+			// "Asset name (asset type, size)"
+			Node->Name = FString::Printf( TEXT( "%s  (%s, %s)" ),
+				*NodeSizeMapData.AssetData.AssetName.ToString(),
+				*NodeSizeMapData.AssetData.AssetClass.ToString(),
+				*SizeMapInternals::MakeBestSizeString( SubtreeSize + NodeSizeMapData.AssetSize, !bAnyUnknownSizesInSubtree && NodeSizeMapData.bHasKnownSize ) );
+
 			const bool bNeedsSelfNode = NodeSizeMapData.AssetSize > 0;
 			if( bNeedsSelfNode )
 			{
-				// "Asset name" (size)
-				Node->Name = FString::Printf( TEXT( "%s (%s)" ),
-					*NodeSizeMapData.AssetData.AssetName.ToString(),
-					*SizeMapInternals::MakeBestSizeString( SubtreeSize + NodeSizeMapData.AssetSize, !bAnyUnknownSizesInSubtree && NodeSizeMapData.bHasKnownSize ) );
-
 				// We have children, so make some space for our own asset's size within our box
 				FTreeMapNodeDataRef ChildSelfTreeMapNode = MakeShareable( new FTreeMapNodeData() );
 				Node->Children.Add( ChildSelfTreeMapNode );
@@ -438,14 +466,6 @@ void SSizeMap::FinalizeNodesRecursively( TSharedPtr<FTreeMapNodeData>& Node, con
 				// Leaf nodes get a background picture
 				ChildSelfTreeMapNode->BackgroundBrush = DefaultThumbnailSlateBrush;
 			}
-			else
-			{
-				// "Asset name (asset type, size)"
-				Node->Name = FString::Printf( TEXT( "%s (%s, %s)" ),
-					*NodeSizeMapData.AssetData.AssetName.ToString(),
-					*NodeSizeMapData.AssetData.AssetClass.ToString(),
-					*SizeMapInternals::MakeBestSizeString( SubtreeSize + NodeSizeMapData.AssetSize, !bAnyUnknownSizesInSubtree && NodeSizeMapData.bHasKnownSize ) );
-			}
 		}
 
 	}
@@ -473,7 +493,7 @@ void SSizeMap::RefreshMap()
 	TSharedPtr<FTreeMapNodeData> SharedRootNode;
 	int32 NumAssetsWhichFailedToLoad = 0;
 	FAssetRegistryModule& AssetRegistryModule = FModuleManager::LoadModuleChecked<FAssetRegistryModule>("AssetRegistry");
-	GatherDependenciesRecursively( AssetRegistryModule, AssetThumbnailPool, VisitedAssetPackageNames, RootAssetPackageNames, RootTreeMapNode, SharedRootNode, NumAssetsWhichFailedToLoad );
+	GatherDependenciesRecursively( AssetRegistryModule, AssetThumbnailPool, VisitedAssetPackageNames, RootAssetPackageNames, RootAssetPackageNames, RootTreeMapNode, SharedRootNode, NumAssetsWhichFailedToLoad );
 
 
 	// Next, do another pass over our tree to and count how big the assets are and to set the node labels.  Also in this pass, we may
@@ -507,7 +527,7 @@ void SSizeMap::RefreshMap()
 		}
 
 		// Use a more descriptive name for the root level node
-		RootTreeMapNode->Name = FString::Printf( TEXT( "%s %s (%i %s)" ),
+		RootTreeMapNode->Name = FString::Printf( TEXT( "%s %s  (%i %s)" ),
 			*LOCTEXT( "RootNode_SizeMapForOneAsset", "Size map for" ).ToString(),
 			*OnlyChild->Name,
 			TotalAssetCount,
@@ -519,7 +539,7 @@ void SSizeMap::RefreshMap()
 		RootTreeMapNode->BackgroundBrush = nullptr;
 		RootTreeMapNode->Size = 0.0f;
 		RootTreeMapNode->Parent = nullptr;
-		RootTreeMapNode->Name = FString::Printf( TEXT( "%s %i %s (%i %s, %s)" ),
+		RootTreeMapNode->Name = FString::Printf( TEXT( "%s %i %s  (%i %s, %s)" ),
 			*LOCTEXT( "RootNode_SizeMapForMultiple", "Size map for" ).ToString(),
 			RootAssetPackageNames.Num(),
 			*LOCTEXT( "RootNode_Assets", "assets" ).ToString(),
@@ -570,7 +590,7 @@ void SSizeMap::OnTreeMapNodeDoubleClicked( FTreeMapNodeData& TreeMapNodeData )
 // @todo sizemap: Should we show the percentage of total size as an option in the UI (though, the size of the boxes show this pretty well.)
 // @todo sizemap: Should we show the folder part of the asset name as a tool-tip?
 // @todo sizemap: It might be nice to unload assets that were loaded by us after the size map is built up
-// @todo sizemap urgent: Trying size map for all engine content BROKE the window updating
+// @todo sizemap: Add a Refresh button. (Also, try to detect when objects change and auto-refresh, optionally)
 
 
 #undef LOCTEXT_NAMESPACE
\ No newline at end of file
diff --git a/Engine/Source/Editor/SizeMap/SSizeMap.h b/Engine/Source/Editor/SizeMap/SSizeMap.h
--- a/Engine/Source/Editor/SizeMap/SSizeMap.h
+++ b/Engine/Source/Editor/SizeMap/SSizeMap.h
@@ -36,7 +36,7 @@ protected:
 	void OnInitialAssetRegistrySearchComplete();
 
 	/** Recursively discovers and loads dependent assets, building up a tree map node hierarchy */
-	void GatherDependenciesRecursively( class FAssetRegistryModule& AssetRegistryModule, TSharedPtr<class FAssetThumbnailPool>& AssetThumbnailPool, TMap<FName, TSharedPtr<class FTreeMapNodeData>>& VisitedAssetPackageNames, const TArray<FName>& AssetPackageNames, const TSharedPtr<FTreeMapNodeData>& ParentTreeMapNode, TSharedPtr<FTreeMapNodeData>& SharedRootNode, int32& NumAssetsWhichFailedToLoad );
+	void GatherDependenciesRecursively( class FAssetRegistryModule& AssetRegistryModule, TSharedPtr<class FAssetThumbnailPool>& AssetThumbnailPool, TMap<FName, TSharedPtr<class FTreeMapNodeData>>& VisitedAssetPackageNames, const TArray<FName>& RootAssetPackageNames, const TArray<FName>& AssetPackageNames, const TSharedPtr<FTreeMapNodeData>& ParentTreeMapNode, TSharedPtr<FTreeMapNodeData>& SharedRootNode, int32& NumAssetsWhichFailedToLoad );
 
 	/** After the node tree is built up, this function is called to generate nice labels for the nodes and to do a final clean-up pass on the tree */
 	void FinalizeNodesRecursively( TSharedPtr<FTreeMapNodeData>& Node, const TSharedPtr<FTreeMapNodeData>& SharedRootNode, int32& TotalAssetCount, SIZE_T& TotalSize, bool& bAnyUnknownSizes );
