Size Mapper fixes
- Fixed assertion when trying to view a size map of a C++ class (UE-14478)
- Fixes assertion when content has direct references script packages

diff --git a/Engine/Source/Editor/SizeMap/SSizeMap.cpp b/Engine/Source/Editor/SizeMap/SSizeMap.cpp
--- a/Engine/Source/Editor/SizeMap/SSizeMap.cpp
+++ b/Engine/Source/Editor/SizeMap/SSizeMap.cpp
@@ -93,6 +93,7 @@ namespace SizeMapInternals
 					AllVisitedObjects.Add( Object );
 
 					const bool bIsAsset =
+						Object->GetOuter() != nullptr &&						// Not a package itself (such as a script package like '/Script/Engine')
 						Object->GetOuter()->IsA( UPackage::StaticClass() ) &&	// Only want outer assets (these should be the only public assets, anyway)
 						Object->HasAllFlags( RF_Public );						// Assets should be public
 
@@ -515,21 +516,27 @@ void SSizeMap::RefreshMap()
 	else if( RootAssetPackageNames.Num() == 1 && !SharedRootNode.IsValid() )
 	{
 		// @todo sizemap: When zoomed right into one asset, can we use the Class color for the node instead of grey?
-		
-		// The root will only have one child, so go ahead and use that child as the actual root
-		FTreeMapNodeDataPtr OnlyChild = RootTreeMapNode->Children[0];
-		OnlyChild->CopyNodeInto( *RootTreeMapNode );
-		RootTreeMapNode->Children = OnlyChild->Children;
-		RootTreeMapNode->Parent = nullptr;
-		for( const auto& ChildNode : RootTreeMapNode->Children )
+
+		FString OnlyAssetName = RootAssetPackageNames[ 0 ].ToString();
+		if( RootTreeMapNode->Children.Num() > 0 )
 		{
-			ChildNode->Parent = RootTreeMapNode.Get();
+			// The root will only have one child, so go ahead and use that child as the actual root
+			FTreeMapNodeDataPtr OnlyChild = RootTreeMapNode->Children[ 0 ];
+			OnlyChild->CopyNodeInto( *RootTreeMapNode );
+			RootTreeMapNode->Children = OnlyChild->Children;
+			RootTreeMapNode->Parent = nullptr;
+			for( const auto& ChildNode : RootTreeMapNode->Children )
+			{
+				ChildNode->Parent = RootTreeMapNode.Get();
+			}
+
+			OnlyAssetName = OnlyChild->Name;
 		}
 
 		// Use a more descriptive name for the root level node
 		RootTreeMapNode->Name = FString::Printf( TEXT( "%s %s  (%i %s)" ),
 			*LOCTEXT( "RootNode_SizeMapForOneAsset", "Size map for" ).ToString(),
-			*OnlyChild->Name,
+			*OnlyAssetName,
 			TotalAssetCount,
 			*LOCTEXT( "RootNode_References", "total assets" ).ToString() );
 	}
