Size Mapper fixes
- Fixed variable shadowing compile error

diff --git a/Engine/Source/Editor/SizeMap/SSizeMap.cpp b/Engine/Source/Editor/SizeMap/SSizeMap.cpp
--- a/Engine/Source/Editor/SizeMap/SSizeMap.cpp
+++ b/Engine/Source/Editor/SizeMap/SSizeMap.cpp
@@ -176,7 +176,7 @@ namespace SizeMapInternals
 }
 
 
-void SSizeMap::GatherDependenciesRecursively( FAssetRegistryModule& AssetRegistryModule, TSharedPtr<FAssetThumbnailPool>& InAssetThumbnailPool, TMap<FName, TSharedPtr<FTreeMapNodeData>>& VisitedAssetPackageNames, const TArray<FName>& RootAssetPackageNames, const TArray<FName>& AssetPackageNames, const TSharedPtr<FTreeMapNodeData>& Node, TSharedPtr<FTreeMapNodeData>& SharedRootNode, int32& NumAssetsWhichFailedToLoad )
+void SSizeMap::GatherDependenciesRecursively( FAssetRegistryModule& AssetRegistryModule, TSharedPtr<FAssetThumbnailPool>& InAssetThumbnailPool, TMap<FName, TSharedPtr<FTreeMapNodeData>>& VisitedAssetPackageNames, const TArray<FName>& AssetPackageNames, const TSharedPtr<FTreeMapNodeData>& Node, TSharedPtr<FTreeMapNodeData>& SharedRootNode, int32& NumAssetsWhichFailedToLoad )
 {
 	for( const FName AssetPackageName : AssetPackageNames )
 	{
@@ -332,7 +332,7 @@ void SSizeMap::GatherDependenciesRecursively( FAssetRegistryModule& AssetRegistr
 
 								
 						// Now visit all of the assets that we are referencing
-						GatherDependenciesRecursively( AssetRegistryModule, InAssetThumbnailPool, VisitedAssetPackageNames, RootAssetPackageNames, ReferencedAssetPackageNames, ChildTreeMapNode, SharedRootNode, NumAssetsWhichFailedToLoad );
+						GatherDependenciesRecursively( AssetRegistryModule, InAssetThumbnailPool, VisitedAssetPackageNames, ReferencedAssetPackageNames, ChildTreeMapNode, SharedRootNode, NumAssetsWhichFailedToLoad );
 					}
 					else
 					{
@@ -494,7 +494,7 @@ void SSizeMap::RefreshMap()
 	TSharedPtr<FTreeMapNodeData> SharedRootNode;
 	int32 NumAssetsWhichFailedToLoad = 0;
 	FAssetRegistryModule& AssetRegistryModule = FModuleManager::LoadModuleChecked<FAssetRegistryModule>("AssetRegistry");
-	GatherDependenciesRecursively( AssetRegistryModule, AssetThumbnailPool, VisitedAssetPackageNames, RootAssetPackageNames, RootAssetPackageNames, RootTreeMapNode, SharedRootNode, NumAssetsWhichFailedToLoad );
+	GatherDependenciesRecursively( AssetRegistryModule, AssetThumbnailPool, VisitedAssetPackageNames, RootAssetPackageNames, RootTreeMapNode, SharedRootNode, NumAssetsWhichFailedToLoad );
 
 
 	// Next, do another pass over our tree to and count how big the assets are and to set the node labels.  Also in this pass, we may
diff --git a/Engine/Source/Editor/SizeMap/SSizeMap.h b/Engine/Source/Editor/SizeMap/SSizeMap.h
--- a/Engine/Source/Editor/SizeMap/SSizeMap.h
+++ b/Engine/Source/Editor/SizeMap/SSizeMap.h
@@ -36,7 +36,7 @@ protected:
 	void OnInitialAssetRegistrySearchComplete();
 
 	/** Recursively discovers and loads dependent assets, building up a tree map node hierarchy */
-	void GatherDependenciesRecursively( class FAssetRegistryModule& AssetRegistryModule, TSharedPtr<class FAssetThumbnailPool>& AssetThumbnailPool, TMap<FName, TSharedPtr<class FTreeMapNodeData>>& VisitedAssetPackageNames, const TArray<FName>& RootAssetPackageNames, const TArray<FName>& AssetPackageNames, const TSharedPtr<FTreeMapNodeData>& ParentTreeMapNode, TSharedPtr<FTreeMapNodeData>& SharedRootNode, int32& NumAssetsWhichFailedToLoad );
+	void GatherDependenciesRecursively( class FAssetRegistryModule& AssetRegistryModule, TSharedPtr<class FAssetThumbnailPool>& AssetThumbnailPool, TMap<FName, TSharedPtr<class FTreeMapNodeData>>& VisitedAssetPackageNames, const TArray<FName>& AssetPackageNames, const TSharedPtr<FTreeMapNodeData>& ParentTreeMapNode, TSharedPtr<FTreeMapNodeData>& SharedRootNode, int32& NumAssetsWhichFailedToLoad );
 
 	/** After the node tree is built up, this function is called to generate nice labels for the nodes and to do a final clean-up pass on the tree */
 	void FinalizeNodesRecursively( TSharedPtr<FTreeMapNodeData>& Node, const TSharedPtr<FTreeMapNodeData>& SharedRootNode, int32& TotalAssetCount, SIZE_T& TotalSize, bool& bAnyUnknownSizes );
